UFS - Added missing Messaging dependency

diff --git a/Engine/Source/Programs/UnrealFileServer/UnrealFileServer.Build.cs b/Engine/Source/Programs/UnrealFileServer/UnrealFileServer.Build.cs
--- a/Engine/Source/Programs/UnrealFileServer/UnrealFileServer.Build.cs
+++ b/Engine/Source/Programs/UnrealFileServer/UnrealFileServer.Build.cs
@@ -19,6 +19,7 @@ public class UnrealFileServer : ModuleRules
                 "SandboxFile",
 				"Sockets",
 				"Projects",
+				"Messaging",
            }
 		);
 	}
